Added Runtime Telemetry Plugin that can record selected telemetry events from Clients and Server builds via StudioTelemetry.
Example demonstrates the recording of IoStoreOnDemand event data on succesful shutdown of the game.

diff --git a/Engine/Plugins/Experimental/RuntimeTelemetry/RuntimeTelemetry.uplugin b/Engine/Plugins/Experimental/RuntimeTelemetry/RuntimeTelemetry.uplugin
new file mode 100644
--- /dev/null
+++ b/Engine/Plugins/Experimental/RuntimeTelemetry/RuntimeTelemetry.uplugin
@@ -0,0 +1,30 @@
+{
+	"FileVersion": 1,
+	"Version": 1,
+	"VersionName": "1.0",
+	"FriendlyName": "Editor Telemetry",
+	"Description": "Plugin that emits common telemetry events from the Editor",
+	"Category": "Telemetry",
+	"CreatedBy": "Epic Games, Inc.",
+	"CreatedByURL": "https://epicgames.com",
+	"DocsURL": "",
+	"MarketplaceURL": "",
+	"SupportURL": "",
+	"EnabledByDefault": true,
+	"CanContainContent": false,
+	"IsBetaVersion": false,
+	"Installed": false,
+	"Modules": [
+		{
+			"Name": "RuntimeTelemetry",
+			"Type": "RuntimeAndProgram",
+			"LoadingPhase": "PostConfigInit"
+		}
+	],
+	"Plugins": [
+		{
+			"Name": "StudioTelemetry",
+			"Enabled": true
+		}
+	]
+}
\ No newline at end of file
diff --git a/Engine/Plugins/Experimental/RuntimeTelemetry/Source/Private/RuntimeTelemetry.cpp b/Engine/Plugins/Experimental/RuntimeTelemetry/Source/Private/RuntimeTelemetry.cpp
new file mode 100644
--- /dev/null
+++ b/Engine/Plugins/Experimental/RuntimeTelemetry/Source/Private/RuntimeTelemetry.cpp
@@ -0,0 +1,56 @@
+// Copyright Epic Games, Inc. All Rights Reserved.
+
+#include "RuntimeTelemetry.h"
+#include "StudioTelemetry.h"
+#include "Misc/CoreDelegates.h"
+
+#ifdef USE_IOSTOREONDEMAND
+#include "IO/IoStoreOnDemand.h"
+#endif
+
+FRuntimeTelemetry& FRuntimeTelemetry::Get()
+{
+	static FRuntimeTelemetry RuntimeTelemetryInstance;
+	return RuntimeTelemetryInstance;
+}
+
+void FRuntimeTelemetry::RecordEvent_IoStoreOnDemand(const FString& Context)
+{
+#ifdef USE_IOSTOREONDEMAND
+	if (FStudioTelemetry::IsAvailable())
+	{
+		const int SchemaVersion = 1;
+		using namespace UE::IoStore;
+
+		if (FIoStoreOnDemandModule* OnDemandModule = FModuleManager::Get().GetModulePtr<FIoStoreOnDemandModule>("IoStoreOnDemand"))
+		{
+			// Only record event data if IoStoreOnDemand is enabled
+			if (OnDemandModule->IsEnabled())
+			{
+				TArray<FAnalyticsEventAttribute> Attributes;
+				OnDemandModule->ReportAnalytics(Attributes);
+
+				Attributes.Emplace(TEXT("SchemaVersion"), SchemaVersion);
+				Attributes.Emplace(TEXT("Context"), Context);
+
+				FStudioTelemetry::Get().RecordEvent("Core.IoStoreOnDemand", Attributes);
+				FStudioTelemetry::Get().FlushEvents();
+			}
+		}
+	}
+#endif
+}
+
+void FRuntimeTelemetry::StartSession()
+{
+	// Register a call back to intercept the point where the game is shutting down
+	FCoreDelegates::OnEnginePreExit.AddLambda([this]()
+		{
+			RecordEvent_IoStoreOnDemand(TEXT("EnginePreExit"));
+		}
+	);
+}
+
+void FRuntimeTelemetry::EndSession()
+{
+}
diff --git a/Engine/Plugins/Experimental/RuntimeTelemetry/Source/Private/RuntimeTelemetryModule.cpp b/Engine/Plugins/Experimental/RuntimeTelemetry/Source/Private/RuntimeTelemetryModule.cpp
new file mode 100644
--- /dev/null
+++ b/Engine/Plugins/Experimental/RuntimeTelemetry/Source/Private/RuntimeTelemetryModule.cpp
@@ -0,0 +1,33 @@
+// Copyright Epic Games, Inc. All Rights Reserved.
+
+#include "RuntimeTelemetryModule.h"
+#include "RuntimeTelemetry.h"
+#include "StudioTelemetry.h"
+#include "Modules/ModuleManager.h"
+
+IMPLEMENT_MODULE(FRuntimeTelemetryModule, RuntimeTelemetry);
+
+void FRuntimeTelemetryModule::StartupModule()
+{
+	if (FStudioTelemetry::Get().IsAvailable())
+	{
+		FRuntimeTelemetry::Get().StartSession();
+	}
+
+	FStudioTelemetry::Get().GetOnStartSession().AddLambda([this]()
+		{
+			FRuntimeTelemetry::Get().StartSession();
+		}
+	);
+	
+	FStudioTelemetry::Get().GetOnEndSession().AddLambda([this]()
+		{
+			FRuntimeTelemetry::Get().EndSession();
+		}
+	);
+}
+
+void FRuntimeTelemetryModule::ShutdownModule()
+{
+	
+}
diff --git a/Engine/Plugins/Experimental/RuntimeTelemetry/Source/Public/RuntimeTelemetry.h b/Engine/Plugins/Experimental/RuntimeTelemetry/Source/Public/RuntimeTelemetry.h
new file mode 100644
--- /dev/null
+++ b/Engine/Plugins/Experimental/RuntimeTelemetry/Source/Public/RuntimeTelemetry.h
@@ -0,0 +1,30 @@
+// Copyright Epic Games, Inc. All Rights Reserved.
+
+#pragma once
+
+#include "Interfaces/IAnalyticsProvider.h"
+#include "Interfaces/IAnalyticsTracer.h"
+#include "CoreMinimal.h"
+#include "Templates/SharedPointer.h"
+
+/**
+ * 
+ */
+class RUNTIMETELEMETRY_API FRuntimeTelemetry
+{
+public:
+	FRuntimeTelemetry() {};
+	~FRuntimeTelemetry() {};
+
+	static FRuntimeTelemetry& Get();
+
+	void StartSession();
+	void EndSession();
+
+private:	
+
+	void RecordEvent_IoStoreOnDemand(const FString& Context);
+};
+
+
+
diff --git a/Engine/Plugins/Experimental/RuntimeTelemetry/Source/Public/RuntimeTelemetryModule.h b/Engine/Plugins/Experimental/RuntimeTelemetry/Source/Public/RuntimeTelemetryModule.h
new file mode 100644
--- /dev/null
+++ b/Engine/Plugins/Experimental/RuntimeTelemetry/Source/Public/RuntimeTelemetryModule.h
@@ -0,0 +1,16 @@
+// Copyright Epic Games, Inc. All Rights Reserved.
+
+#pragma once
+
+#include "CoreMinimal.h"
+#include "Modules/ModuleInterface.h"
+
+class RUNTIMETELEMETRY_API FRuntimeTelemetryModule : public IModuleInterface
+{
+public:
+	// IModuleInterface
+	virtual void StartupModule() override;
+	virtual void ShutdownModule() override;
+
+private:
+};
diff --git a/Engine/Plugins/Experimental/RuntimeTelemetry/Source/RuntimeTelemetry.Build.cs b/Engine/Plugins/Experimental/RuntimeTelemetry/Source/RuntimeTelemetry.Build.cs
new file mode 100644
--- /dev/null
+++ b/Engine/Plugins/Experimental/RuntimeTelemetry/Source/RuntimeTelemetry.Build.cs
@@ -0,0 +1,31 @@
+// Copyright Epic Games, Inc. All Rights Reserved.
+
+namespace UnrealBuildTool.Rules
+{
+	public class RuntimeTelemetry : ModuleRules
+	{
+		public RuntimeTelemetry(ReadOnlyTargetRules Target) : base(Target)
+		{
+			PublicDependencyModuleNames.AddRange(
+				new string[] {
+					"Core",
+					"Engine"
+				}
+			);
+
+			PrivateDependencyModuleNames.AddRange(
+				new string[] {
+					"CoreUObject",
+					"Analytics",
+					"StudioTelemetry"
+				}
+			);
+
+			if (Target.bCompileIoStoreOnDemand && !Target.bBuildEditor)
+			{
+				PrivateDependencyModuleNames.Add("IoStoreOnDemand");
+				PublicDefinitions.Add("USE_IOSTOREONDEMAND");
+			}
+		}
+	}
+}
