Removed EditorTelemetry and RutnimeTelemetry from Shipping Builds

diff --git a/Engine/Plugins/Experimental/EditorTelemetry/Source/Private/EditorTelemetryModule.cpp b/Engine/Plugins/Experimental/EditorTelemetry/Source/Private/EditorTelemetryModule.cpp
--- a/Engine/Plugins/Experimental/EditorTelemetry/Source/Private/EditorTelemetryModule.cpp
+++ b/Engine/Plugins/Experimental/EditorTelemetry/Source/Private/EditorTelemetryModule.cpp
@@ -9,6 +9,7 @@ IMPLEMENT_MODULE(FEditorTelemetryModule, EditorTelemetry);
 
 void FEditorTelemetryModule::StartupModule()
 {
+#if !UE_BUILD_SHIPPING
 	if (FStudioTelemetry::Get().IsSessionRunning())
 	{
 		// StudioTelemetry session has already started
@@ -28,6 +29,7 @@ void FEditorTelemetryModule::StartupModule()
 			FEditorTelemetry::Get().EndSession();
 		}
 	);
+#endif
 }
 
 void FEditorTelemetryModule::ShutdownModule()
diff --git a/Engine/Plugins/Experimental/RuntimeTelemetry/Source/Private/RuntimeTelemetryModule.cpp b/Engine/Plugins/Experimental/RuntimeTelemetry/Source/Private/RuntimeTelemetryModule.cpp
--- a/Engine/Plugins/Experimental/RuntimeTelemetry/Source/Private/RuntimeTelemetryModule.cpp
+++ b/Engine/Plugins/Experimental/RuntimeTelemetry/Source/Private/RuntimeTelemetryModule.cpp
@@ -9,6 +9,7 @@ IMPLEMENT_MODULE(FRuntimeTelemetryModule, RuntimeTelemetry);
 
 void FRuntimeTelemetryModule::StartupModule()
 {
+#if !UE_BUILD_SHIPPING
 	if (FStudioTelemetry::Get().IsSessionRunning())
 	{
 		// StudioTelemetry session has already started
@@ -28,6 +29,7 @@ void FRuntimeTelemetryModule::StartupModule()
 			FRuntimeTelemetry::Get().EndSession();
 		}
 	);
+#endif
 }
 
 void FRuntimeTelemetryModule::ShutdownModule()
