Minor tweaks to Runtime Telemetry and Editor Telemetry plugins
Made the RecordEvent APIs for Editor and Runtime telemetry public
Added LLM events to Runtime Telemetry

diff --git a/Engine/Plugins/Experimental/EditorPerformance/Source/EditorPerformance/Private/EditorPerformanceModule.cpp b/Engine/Plugins/Experimental/EditorPerformance/Source/EditorPerformance/Private/EditorPerformanceModule.cpp
--- a/Engine/Plugins/Experimental/EditorPerformance/Source/EditorPerformance/Private/EditorPerformanceModule.cpp
+++ b/Engine/Plugins/Experimental/EditorPerformance/Source/EditorPerformance/Private/EditorPerformanceModule.cpp
@@ -113,7 +113,6 @@ void FEditorPerformanceModule::InitializeEditor()
 			CloudDDCReadSpeedKPI		= KPIRegistry.DeclareKPIValue(CacheCategoryName, CloudDDCReadSpeedKPIName, 100.0, CloudDDCReadSpeedKPILimit, FKPIValue::GreaterThan, FKPIValue::MegaBitsPerSecond);
 			EditorStallRateKPI			= KPIRegistry.DeclareKPIValue(EditorCategoryName, EditorStallRateKPIName, 0.0, EditorStallRateKPILimit, FKPIValue::LessThan, FKPIValue::Percent);
 			PIEStallRateKPI				= KPIRegistry.DeclareKPIValue(PIECategoryName, PIEStallRateKPIName, 0.0, PIEStallRateKPILimit, FKPIValue::LessThan, FKPIValue::Percent);
-			VirtualAssetEfficiencyKPI	= KPIRegistry.DeclareKPIValue(CacheCategoryName, VirtualAssetEfficiencyKPIName, 100.0, VirtualAssetEfficiencyKPILimit, FKPIValue::GreaterThan, FKPIValue::Percent);
 		}
 
 		// Populate the notification list with all KPI values if it is empty. 
@@ -217,8 +216,8 @@ void FEditorPerformanceModule::InitializeKPIs()
 	PIEShutdownKPI				= KPIRegistry.DeclareKPIValue(PIECategoryName, PIEShutdownKPIName, 0.0, PIEShutdownKPILimit, FKPIValue::LessThan, FKPIValue::Minutes);
 	PIEHitchRateKPI				= KPIRegistry.DeclareKPIValue(PIECategoryName, PIEHitchRateKPIName, 0.0, PIEHitchRateKPILimit, FKPIValue::LessThan, FKPIValue::Percent);
 	TotalTimeToPIEKPI			= KPIRegistry.DeclareKPIValue(PIECategoryName, TotalTimeToPIEKPIName, 0.0, TotalTimeToPIEKPILimit, FKPIValue::LessThan, FKPIValue::Minutes);
-	//TotalDDCEfficiencyKPI		= KPIRegistry.DeclareKPIValue(CacheCategoryName, TotalDDCEfficiencyKPIName, 100.0, TotalDDCEffciencyKPILimit, FKPIValue::GreaterThan, FKPIValue::Percent);
 	LocalDDCEfficiencyKPI		= KPIRegistry.DeclareKPIValue(CacheCategoryName, LocalDDCEfficiencyKPIName, 100.0, LocalDDCEffciencyKPILimit, FKPIValue::GreaterThan, FKPIValue::Percent);
+	VirtualAssetEfficiencyKPI	= KPIRegistry.DeclareKPIValue(CacheCategoryName, VirtualAssetEfficiencyKPIName, 100.0, VirtualAssetEfficiencyKPILimit, FKPIValue::GreaterThan, FKPIValue::Percent);
 	CoreCountKPI				= KPIRegistry.DeclareKPIValue(HardwareCategoryName, CoreCountKPIName, 128.0, CoreCountKPILimit, FKPIValue::GreaterThanOrEqual, FKPIValue::Decimal);
 	TotalMemoryKPI				= KPIRegistry.DeclareKPIValue(HardwareCategoryName, TotalMemoryKPIName, 128.0, TotalMemoryKPILimit, FKPIValue::GreaterThanOrEqual, FKPIValue::GigaBytes);
 	AvailableMemoryKPI			= KPIRegistry.DeclareKPIValue(HardwareCategoryName, AvailableMemoryKPIName, 128.0, AvailableMemoryKPILimit, FKPIValue::GreaterThanOrEqual, FKPIValue::GigaBytes);
diff --git a/Engine/Plugins/Experimental/EditorTelemetry/Source/Private/EditorTelemetry.cpp b/Engine/Plugins/Experimental/EditorTelemetry/Source/Private/EditorTelemetry.cpp
--- a/Engine/Plugins/Experimental/EditorTelemetry/Source/Private/EditorTelemetry.cpp
+++ b/Engine/Plugins/Experimental/EditorTelemetry/Source/Private/EditorTelemetry.cpp
@@ -999,20 +999,7 @@ void FEditorTelemetry::StartSession()
 
 				if (EnumHasAllFlags(Flags, UE::Virtualization::EAnalyticsFlags::Flush))
 				{
-					if (TSharedPtr<IAnalyticsProvider> Provider = FStudioTelemetry::Get().GetProvider().Pin())
-					{
-						Provider->FlushEvents();
-
-						// It is quite likely that one of the analytics providers is sending data via the FHttpManager so we
-						// need to flush that as well to make sure that the message is sent. If the virtualization system has
-						// requested that we flush then it is quite likely that the process is to be terminated in which case
-						// not flushing the FHttpManager will result in analytics not being properly recorded.
-						if (FHttpModule* HttpModule = FModuleManager::GetModulePtr<FHttpModule>("HTTP"))
-						{
-							FHttpManager& HttpManager = HttpModule->GetHttpManager();
-							HttpManager.Flush(EHttpFlushReason::FullFlush);
-						}	
-					}
+					FStudioTelemetry::Get().FlushEvents();
 				}
 			});
 	}
diff --git a/Engine/Plugins/Experimental/EditorTelemetry/Source/Public/EditorTelemetry.h b/Engine/Plugins/Experimental/EditorTelemetry/Source/Public/EditorTelemetry.h
--- a/Engine/Plugins/Experimental/EditorTelemetry/Source/Public/EditorTelemetry.h
+++ b/Engine/Plugins/Experimental/EditorTelemetry/Source/Public/EditorTelemetry.h
@@ -20,26 +20,27 @@ struct FTimerHandle;
 class EDITORTELEMETRY_API FEditorTelemetry
 {
 public:
-	FEditorTelemetry() {};
-	~FEditorTelemetry() {};
+
+	FEditorTelemetry() = default;
+	~FEditorTelemetry() = default;
 
 	static FEditorTelemetry& Get();
 
 	void StartSession();
 	void EndSession();
 
+	/** Useful event recording functions */
+	void RecordEvent_Cooking(TArray<FAnalyticsEventAttribute> Attributes = {});
+	void RecordEvent_Loading(const FString& Context, double LoadingSeconds, TArray<FAnalyticsEventAttribute> Attributes = {});
+	void RecordEvent_CoreSystems(const FString& Context, TArray<FAnalyticsEventAttribute> Attributes = {});
+	void RecordEvent_DDCResource(const FString& Context, TArray<FAnalyticsEventAttribute> Attributes = {});
+	void RecordEvent_DDCSummary(const FString& Context, TArray<FAnalyticsEventAttribute> Attributes = {});
+	void RecordEvent_Zen(const FString& Context, TArray<FAnalyticsEventAttribute> Attributes = {});
+	void RecordEvent_VirtualAssets(const FString& Context, TArray<FAnalyticsEventAttribute> Attributes = {});
+	void RegisterCollectionWorkflowDelegates(FTelemetryRouter& Router);
+
 private:	
 
-	/** Various event handling functions */
-	static void RecordEvent_Cooking(TArray<FAnalyticsEventAttribute> Attributes = {});
-	static void RecordEvent_Loading(const FString& LoadingName, double LoadingSeconds, TArray<FAnalyticsEventAttribute> Attributes = {});
-	static void RecordEvent_CoreSystems(const FString& Context, TArray<FAnalyticsEventAttribute> Attributes = {});
-	static void RecordEvent_DDCResource(const FString& Context, TArray<FAnalyticsEventAttribute> Attributes = {});
-	static void RecordEvent_DDCSummary(const FString& Context, TArray<FAnalyticsEventAttribute> Attributes = {});
-	static void RecordEvent_IAS(const FString& Context, TArray<FAnalyticsEventAttribute> Attributes = {});
-	static void RecordEvent_Zen(const FString& Context, TArray<FAnalyticsEventAttribute> Attributes = {});
-	static void RecordEvent_VirtualAssets(const FString& Context, TArray<FAnalyticsEventAttribute> Attributes = {});
-	static void RegisterCollectionWorkflowDelegates(FTelemetryRouter& Router);
 	void HitchSamplerCallback();
 	void HeartbeatCallback();
 	
diff --git a/Engine/Plugins/Experimental/RuntimeTelemetry/Source/Private/RuntimeTelemetry.cpp b/Engine/Plugins/Experimental/RuntimeTelemetry/Source/Private/RuntimeTelemetry.cpp
--- a/Engine/Plugins/Experimental/RuntimeTelemetry/Source/Private/RuntimeTelemetry.cpp
+++ b/Engine/Plugins/Experimental/RuntimeTelemetry/Source/Private/RuntimeTelemetry.cpp
@@ -3,11 +3,14 @@
 #include "RuntimeTelemetry.h"
 #include "StudioTelemetry.h"
 #include "Misc/CoreDelegates.h"
+#include "Containers/Map.h"
+#include "HAL/LowLevelMemTracker.h"
 
 #ifdef USE_IOSTOREONDEMAND
 #include "IO/IoStoreOnDemand.h"
 #endif
 
+
 FRuntimeTelemetry& FRuntimeTelemetry::Get()
 {
 	static FRuntimeTelemetry RuntimeTelemetryInstance;
@@ -16,17 +19,18 @@ FRuntimeTelemetry& FRuntimeTelemetry::Get()
 
 void FRuntimeTelemetry::RecordEvent_IoStoreOnDemand(const FString& Context)
 {
-#ifdef USE_IOSTOREONDEMAND
-	if (FStudioTelemetry::IsAvailable())
-	{
-		const int SchemaVersion = 1;
-		using namespace UE::IoStore;
+#ifdef USE_IOSTOREONDEMAND	
+	using namespace UE::IoStore;
 
+	if (FStudioTelemetry::Get().IsSessionRunning())
+	{
 		if (FIoStoreOnDemandModule* OnDemandModule = FModuleManager::Get().GetModulePtr<FIoStoreOnDemandModule>("IoStoreOnDemand"))
 		{
 			// Only record event data if IoStoreOnDemand is enabled
 			if (OnDemandModule->IsEnabled())
 			{
+				const int SchemaVersion = 1;
+
 				TArray<FAnalyticsEventAttribute> Attributes;
 				OnDemandModule->ReportAnalytics(Attributes);
 
@@ -41,12 +45,90 @@ void FRuntimeTelemetry::RecordEvent_IoStoreOnDemand(const FString& Context)
 #endif
 }
 
+void FRuntimeTelemetry::RecordEvent_LLMSystemMemory(const FString& Context)
+{
+#if ENABLE_LOW_LEVEL_MEM_TRACKER
+	if (FStudioTelemetry::Get().IsSessionRunning())
+	{
+		TMap<FName, uint64> LLMTrackedMemoryMap;
+		FLowLevelMemTracker::Get().GetTrackedTagsNamesWithAmount(LLMTrackedMemoryMap, ELLMTracker::Default, ELLMTagSet::None);
+
+		for (TMap<FName, uint64>::TConstIterator It(LLMTrackedMemoryMap); It; ++It)
+		{
+			const int SchemaVersion = 1;
+
+			TArray<FAnalyticsEventAttribute> Attributes;
+
+			Attributes.Emplace(TEXT("SchemaVersion"), SchemaVersion);
+			Attributes.Emplace(TEXT("Context"), Context);
+			Attributes.Emplace(TEXT("Name"), It->Key);
+			Attributes.Emplace(TEXT("Size"), It->Value);
+
+			FStudioTelemetry::Get().RecordEvent("Memory.LLM.System", Attributes);
+		}
+	}
+#endif
+}
+
+void FRuntimeTelemetry::RecordEvent_LLMAssetClassMemory(const FString& Context)
+{
+#if ENABLE_LOW_LEVEL_MEM_TRACKER && LLM_ALLOW_ASSETS_TAGS
+	if (FStudioTelemetry::Get().IsSessionRunning())
+	{
+		TMap<FName, uint64> LLMTrackedMemoryMap;
+		FLowLevelMemTracker::Get().GetTrackedTagsNamesWithAmount(LLMTrackedMemoryMap, ELLMTracker::Default, ELLMTagSet::AssetClasses);
+
+		for (TMap<FName, uint64>::TConstIterator It(LLMTrackedMemoryMap); It; ++It)
+		{
+			const int SchemaVersion = 1;
+
+			TArray<FAnalyticsEventAttribute> Attributes;
+
+			Attributes.Emplace(TEXT("SchemaVersion"), SchemaVersion);
+			Attributes.Emplace(TEXT("Context"), Context);
+			Attributes.Emplace(TEXT("Name"), It->Key);
+			Attributes.Emplace(TEXT("Size"), It->Value);
+
+			FStudioTelemetry::Get().RecordEvent("Memory.LLM.AssetClass", Attributes);
+		}
+	}
+#endif
+}
+
+void FRuntimeTelemetry::RecordEvent_LLMAssetMemory(const FString& Context)
+{
+#if ENABLE_LOW_LEVEL_MEM_TRACKER && LLM_ALLOW_ASSETS_TAGS
+	if (FStudioTelemetry::Get().IsSessionRunning())
+	{
+		TMap<FName, uint64> LLMTrackedMemoryMap;
+		FLowLevelMemTracker::Get().GetTrackedTagsNamesWithAmount(LLMTrackedMemoryMap, ELLMTracker::Default, ELLMTagSet::Assets);
+
+		for (TMap<FName, uint64>::TConstIterator It(LLMTrackedMemoryMap); It; ++It)
+		{
+			const int SchemaVersion = 1;
+
+			TArray<FAnalyticsEventAttribute> Attributes;
+
+			Attributes.Emplace(TEXT("SchemaVersion"), SchemaVersion);
+			Attributes.Emplace(TEXT("Context"), Context);
+			Attributes.Emplace(TEXT("Name"), It->Key);
+			Attributes.Emplace(TEXT("Size"), It->Value);
+
+			FStudioTelemetry::Get().RecordEvent("Memory.LLM.Asset", Attributes);
+		}
+	}
+#endif
+}
+
 void FRuntimeTelemetry::StartSession()
 {
 	// Register a call back to intercept the point where the game is shutting down
 	FCoreDelegates::OnEnginePreExit.AddLambda([this]()
 		{
 			RecordEvent_IoStoreOnDemand(TEXT("EnginePreExit"));
+			RecordEvent_LLMSystemMemory(TEXT("EnginePreExit"));
+			RecordEvent_LLMAssetClassMemory(TEXT("EnginePreExit"));
+			RecordEvent_LLMAssetMemory(TEXT("EnginePreExit"));
 		}
 	);
 }
diff --git a/Engine/Plugins/Experimental/RuntimeTelemetry/Source/Public/RuntimeTelemetry.h b/Engine/Plugins/Experimental/RuntimeTelemetry/Source/Public/RuntimeTelemetry.h
--- a/Engine/Plugins/Experimental/RuntimeTelemetry/Source/Public/RuntimeTelemetry.h
+++ b/Engine/Plugins/Experimental/RuntimeTelemetry/Source/Public/RuntimeTelemetry.h
@@ -13,20 +13,22 @@
 class RUNTIMETELEMETRY_API FRuntimeTelemetry
 {
 public:
-	FRuntimeTelemetry() {};
-	~FRuntimeTelemetry() {};
+	FRuntimeTelemetry() = default;
+	~FRuntimeTelemetry() = default;
 
 	static FRuntimeTelemetry& Get();
 
 	void StartSession();
 	void EndSession();
 
-private:	
-
+	/** Useful event recording functions */
 	void RecordEvent_IoStoreOnDemand(const FString& Context);
 	void RecordEvent_LLMSystemMemory(const FString& Context);
 	void RecordEvent_LLMAssetClassMemory(const FString& Context);
 	void RecordEvent_LLMAssetMemory(const FString& Context);
+
+private:	
+
 };
 
 
