Updated LLM Memory Telemetry Event for Rutnime LLM capture

diff --git a/Engine/Plugins/Experimental/RuntimeTelemetry/Source/Private/RuntimeTelemetry.cpp b/Engine/Plugins/Experimental/RuntimeTelemetry/Source/Private/RuntimeTelemetry.cpp
--- a/Engine/Plugins/Experimental/RuntimeTelemetry/Source/Private/RuntimeTelemetry.cpp
+++ b/Engine/Plugins/Experimental/RuntimeTelemetry/Source/Private/RuntimeTelemetry.cpp
@@ -45,77 +45,43 @@ void FRuntimeTelemetry::RecordEvent_IoStoreOnDemand(const FString& Context)
 #endif
 }
 
-void FRuntimeTelemetry::RecordEvent_LLMSystemMemory(const FString& Context)
+void FRuntimeTelemetry::RecordEvent_MemoryLLM(const FString& Context)
 {
 #if ENABLE_LOW_LEVEL_MEM_TRACKER
 	if (FStudioTelemetry::Get().IsSessionRunning())
 	{
-		TMap<FName, uint64> LLMTrackedMemoryMap;
-		FLowLevelMemTracker::Get().GetTrackedTagsNamesWithAmount(LLMTrackedMemoryMap, ELLMTracker::Default, ELLMTagSet::None);
-
-		for (TMap<FName, uint64>::TConstIterator It(LLMTrackedMemoryMap); It; ++It)
-		{
-			const int SchemaVersion = 1;
-
-			TArray<FAnalyticsEventAttribute> Attributes;
-
-			Attributes.Emplace(TEXT("SchemaVersion"), SchemaVersion);
-			Attributes.Emplace(TEXT("Context"), Context);
-			Attributes.Emplace(TEXT("Name"), It->Key);
-			Attributes.Emplace(TEXT("Size"), It->Value);
-
-			FStudioTelemetry::Get().RecordEvent("Memory.LLM.System", Attributes);
-		}
-	}
-#endif
-}
-
-void FRuntimeTelemetry::RecordEvent_LLMAssetClassMemory(const FString& Context)
-{
-#if ENABLE_LOW_LEVEL_MEM_TRACKER && LLM_ALLOW_ASSETS_TAGS
-	if (FStudioTelemetry::Get().IsSessionRunning())
-	{
-		TMap<FName, uint64> LLMTrackedMemoryMap;
-		FLowLevelMemTracker::Get().GetTrackedTagsNamesWithAmount(LLMTrackedMemoryMap, ELLMTracker::Default, ELLMTagSet::AssetClasses);
-
-		for (TMap<FName, uint64>::TConstIterator It(LLMTrackedMemoryMap); It; ++It)
-		{
-			const int SchemaVersion = 1;
-
-			TArray<FAnalyticsEventAttribute> Attributes;
-
-			Attributes.Emplace(TEXT("SchemaVersion"), SchemaVersion);
-			Attributes.Emplace(TEXT("Context"), Context);
-			Attributes.Emplace(TEXT("Name"), It->Key);
-			Attributes.Emplace(TEXT("Size"), It->Value);
-
-			FStudioTelemetry::Get().RecordEvent("Memory.LLM.AssetClass", Attributes);
-		}
-	}
-#endif
-}
-
-void FRuntimeTelemetry::RecordEvent_LLMAssetMemory(const FString& Context)
-{
-#if ENABLE_LOW_LEVEL_MEM_TRACKER && LLM_ALLOW_ASSETS_TAGS
-	if (FStudioTelemetry::Get().IsSessionRunning())
-	{
-		TMap<FName, uint64> LLMTrackedMemoryMap;
-		FLowLevelMemTracker::Get().GetTrackedTagsNamesWithAmount(LLMTrackedMemoryMap, ELLMTracker::Default, ELLMTagSet::Assets);
-
-		for (TMap<FName, uint64>::TConstIterator It(LLMTrackedMemoryMap); It; ++It)
-		{
-			const int SchemaVersion = 1;
-
-			TArray<FAnalyticsEventAttribute> Attributes;
-
-			Attributes.Emplace(TEXT("SchemaVersion"), SchemaVersion);
-			Attributes.Emplace(TEXT("Context"), Context);
-			Attributes.Emplace(TEXT("Name"), It->Key);
-			Attributes.Emplace(TEXT("Size"), It->Value);
-
-			FStudioTelemetry::Get().RecordEvent("Memory.LLM.Asset", Attributes);
-		}
+		auto RecordLLMMemoryEvent = [](const FString& Context, const FString& TagSet, TMap<FName, uint64>& LLMTrackedMemoryMap)
+			{
+				for (TMap<FName, uint64>::TConstIterator It(LLMTrackedMemoryMap); It; ++It)
+				{
+					const int SchemaVersion = 1;
+
+					TArray<FAnalyticsEventAttribute> Attributes;
+
+					Attributes.Emplace(TEXT("SchemaVersion"), SchemaVersion);
+					Attributes.Emplace(TEXT("Context"), Context);
+					Attributes.Emplace(TEXT("TagSet"), TEXT("None"));
+					Attributes.Emplace(TEXT("Name"), It->Key);
+					Attributes.Emplace(TEXT("Size"), It->Value);
+
+					FStudioTelemetry::Get().RecordEvent("Core.Memory.LLM", Attributes);
+				}
+			};
+
+		// None TagSet
+		TMap<FName, uint64> LLMTrackedNoneMemory;
+		FLowLevelMemTracker::Get().GetTrackedTagsNamesWithAmount(LLMTrackedNoneMemory, ELLMTracker::Default, ELLMTagSet::None);
+		RecordLLMMemoryEvent(Context, TEXT("None"), LLMTrackedNoneMemory);
+	
+		// AssetClasses TagSet
+		TMap<FName, uint64> LLMTrackedAssetClassesMemory;
+		FLowLevelMemTracker::Get().GetTrackedTagsNamesWithAmount(LLMTrackedAssetClassesMemory, ELLMTracker::Default, ELLMTagSet::AssetClasses);
+		RecordLLMMemoryEvent(Context, TEXT("AssetClasses"), LLMTrackedAssetClassesMemory);
+		
+		// Asset TagSet	
+		TMap<FName, uint64> LLMTrackedAssetMemory;
+		FLowLevelMemTracker::Get().GetTrackedTagsNamesWithAmount(LLMTrackedAssetMemory, ELLMTracker::Default, ELLMTagSet::Assets);
+		RecordLLMMemoryEvent(Context, TEXT("Assets"), LLMTrackedAssetMemory);
 	}
 #endif
 }
@@ -126,9 +92,7 @@ void FRuntimeTelemetry::StartSession()
 	FCoreDelegates::OnEnginePreExit.AddLambda([this]()
 		{
 			RecordEvent_IoStoreOnDemand(TEXT("EnginePreExit"));
-			RecordEvent_LLMSystemMemory(TEXT("EnginePreExit"));
-			RecordEvent_LLMAssetClassMemory(TEXT("EnginePreExit"));
-			RecordEvent_LLMAssetMemory(TEXT("EnginePreExit"));
+			RecordEvent_MemoryLLM(TEXT("EnginePreExit"));
 		}
 	);
 }
diff --git a/Engine/Plugins/Experimental/RuntimeTelemetry/Source/Public/RuntimeTelemetry.h b/Engine/Plugins/Experimental/RuntimeTelemetry/Source/Public/RuntimeTelemetry.h
--- a/Engine/Plugins/Experimental/RuntimeTelemetry/Source/Public/RuntimeTelemetry.h
+++ b/Engine/Plugins/Experimental/RuntimeTelemetry/Source/Public/RuntimeTelemetry.h
@@ -23,10 +23,8 @@ public:
 
 	/** Useful event recording functions */
 	void RecordEvent_IoStoreOnDemand(const FString& Context);
-	void RecordEvent_LLMSystemMemory(const FString& Context);
-	void RecordEvent_LLMAssetClassMemory(const FString& Context);
-	void RecordEvent_LLMAssetMemory(const FString& Context);
-
+	void RecordEvent_MemoryLLM(const FString& Context);
+	
 private:	
 
 };
