Added LLM Memory Telemetry to Editor and Runtime Telemetry plugins

diff --git a/Engine/Plugins/Experimental/EditorTelemetry/Source/Private/EditorTelemetry.cpp b/Engine/Plugins/Experimental/EditorTelemetry/Source/Private/EditorTelemetry.cpp
--- a/Engine/Plugins/Experimental/EditorTelemetry/Source/Private/EditorTelemetry.cpp
+++ b/Engine/Plugins/Experimental/EditorTelemetry/Source/Private/EditorTelemetry.cpp
@@ -307,12 +307,54 @@ void FEditorTelemetry::RecordEvent_VirtualAssets(const FString& Context, TArray<
 	}
 }
 
+void FEditorTelemetry::RecordEvent_MemoryLLM(const FString& Context, TArray<FAnalyticsEventAttribute> Attributes)
+{
+#if ENABLE_LOW_LEVEL_MEM_TRACKER
+	if (FStudioTelemetry::Get().IsSessionRunning())
+	{
+		auto RecordLLMMemoryEvent = [&Attributes](const FString& Context, const FString& TagSet, TMap<FName, uint64>& LLMTrackedMemoryMap)
+			{
+				for (TMap<FName, uint64>::TConstIterator It(LLMTrackedMemoryMap); It; ++It)
+				{
+					const int SchemaVersion = 2;
+
+					TArray<FAnalyticsEventAttribute> EventAttributes = Attributes;
+
+					EventAttributes.Emplace(TEXT("SchemaVersion"), SchemaVersion);
+					EventAttributes.Emplace(TEXT("Context"), Context);
+					EventAttributes.Emplace(TEXT("TagSet"), TagSet);
+					EventAttributes.Emplace(TEXT("Name"), It->Key);
+					EventAttributes.Emplace(TEXT("Size"), It->Value);
+
+					FStudioTelemetry::Get().RecordEvent("Core.Memory.LLM", EventAttributes);
+				}
+			};
+
+		// None TagSet
+		TMap<FName, uint64> LLMTrackedNoneMemory;
+		FLowLevelMemTracker::Get().GetTrackedTagsNamesWithAmount(LLMTrackedNoneMemory, ELLMTracker::Default, ELLMTagSet::None);
+		RecordLLMMemoryEvent(Context, TEXT("None"), LLMTrackedNoneMemory);
+
+		// AssetClasses TagSet
+		TMap<FName, uint64> LLMTrackedAssetClassesMemory;
+		FLowLevelMemTracker::Get().GetTrackedTagsNamesWithAmount(LLMTrackedAssetClassesMemory, ELLMTracker::Default, ELLMTagSet::AssetClasses);
+		RecordLLMMemoryEvent(Context, TEXT("AssetClasses"), LLMTrackedAssetClassesMemory);
+
+		// Asset TagSet	
+		TMap<FName, uint64> LLMTrackedAssetMemory;
+		FLowLevelMemTracker::Get().GetTrackedTagsNamesWithAmount(LLMTrackedAssetMemory, ELLMTracker::Default, ELLMTagSet::Assets);
+		RecordLLMMemoryEvent(Context, TEXT("Assets"), LLMTrackedAssetMemory);
+	}
+#endif
+}
+
 void FEditorTelemetry::RecordEvent_CoreSystems(const FString& Context, TArray<FAnalyticsEventAttribute> Attributes)
 {
 	FEditorTelemetry::RecordEvent_DDCResource(Context, Attributes);
 	FEditorTelemetry::RecordEvent_DDCSummary(Context, Attributes);
 	FEditorTelemetry::RecordEvent_Zen(Context, Attributes);
 	FEditorTelemetry::RecordEvent_VirtualAssets(Context, Attributes);
+	FEditorTelemetry::RecordEvent_MemoryLLM(Context, Attributes);
 }
 
 void FEditorTelemetry::RegisterCollectionWorkflowDelegates(FTelemetryRouter& Router)
diff --git a/Engine/Plugins/Experimental/EditorTelemetry/Source/Public/EditorTelemetry.h b/Engine/Plugins/Experimental/EditorTelemetry/Source/Public/EditorTelemetry.h
--- a/Engine/Plugins/Experimental/EditorTelemetry/Source/Public/EditorTelemetry.h
+++ b/Engine/Plugins/Experimental/EditorTelemetry/Source/Public/EditorTelemetry.h
@@ -37,6 +37,7 @@ public:
 	void RecordEvent_DDCSummary(const FString& Context, TArray<FAnalyticsEventAttribute> Attributes = {});
 	void RecordEvent_Zen(const FString& Context, TArray<FAnalyticsEventAttribute> Attributes = {});
 	void RecordEvent_VirtualAssets(const FString& Context, TArray<FAnalyticsEventAttribute> Attributes = {});
+	void RecordEvent_MemoryLLM(const FString& Context, TArray<FAnalyticsEventAttribute> Attributes = {});
 	void RegisterCollectionWorkflowDelegates(FTelemetryRouter& Router);
 
 private:	
diff --git a/Engine/Plugins/Experimental/RuntimeTelemetry/Source/Private/RuntimeTelemetry.cpp b/Engine/Plugins/Experimental/RuntimeTelemetry/Source/Private/RuntimeTelemetry.cpp
--- a/Engine/Plugins/Experimental/RuntimeTelemetry/Source/Private/RuntimeTelemetry.cpp
+++ b/Engine/Plugins/Experimental/RuntimeTelemetry/Source/Private/RuntimeTelemetry.cpp
@@ -17,7 +17,7 @@ FRuntimeTelemetry& FRuntimeTelemetry::Get()
 	return RuntimeTelemetryInstance;
 }
 
-void FRuntimeTelemetry::RecordEvent_IoStoreOnDemand(const FString& Context)
+void FRuntimeTelemetry::RecordEvent_IoStoreOnDemand(const FString& Context, TArray<FAnalyticsEventAttribute> Attributes)
 {
 #ifdef USE_IOSTOREONDEMAND	
 	using namespace UE::IoStore;
@@ -31,13 +31,14 @@ void FRuntimeTelemetry::RecordEvent_IoStoreOnDemand(const FString& Context)
 			{
 				const int SchemaVersion = 1;
 
-				TArray<FAnalyticsEventAttribute> Attributes;
-				OnDemandModule->ReportAnalytics(Attributes);
+				TArray<FAnalyticsEventAttribute> EventAttributes = Attributes;
 
-				Attributes.Emplace(TEXT("SchemaVersion"), SchemaVersion);
-				Attributes.Emplace(TEXT("Context"), Context);
+				OnDemandModule->ReportAnalytics(EventAttributes);
 
-				FStudioTelemetry::Get().RecordEvent("Core.IoStoreOnDemand", Attributes);
+				EventAttributes.Emplace(TEXT("SchemaVersion"), SchemaVersion);
+				EventAttributes.Emplace(TEXT("Context"), Context);
+
+				FStudioTelemetry::Get().RecordEvent("Core.IoStoreOnDemand", EventAttributes);
 				FStudioTelemetry::Get().FlushEvents();
 			}
 		}
@@ -45,26 +46,27 @@ void FRuntimeTelemetry::RecordEvent_IoStoreOnDemand(const FString& Context)
 #endif
 }
 
-void FRuntimeTelemetry::RecordEvent_MemoryLLM(const FString& Context)
+
+void FRuntimeTelemetry::RecordEvent_MemoryLLM(const FString& Context, TArray<FAnalyticsEventAttribute> Attributes)
 {
 #if ENABLE_LOW_LEVEL_MEM_TRACKER
 	if (FStudioTelemetry::Get().IsSessionRunning())
 	{
-		auto RecordLLMMemoryEvent = [](const FString& Context, const FString& TagSet, TMap<FName, uint64>& LLMTrackedMemoryMap)
+		auto RecordLLMMemoryEvent = [&Attributes](const FString& Context, const FString& TagSet, TMap<FName, uint64>& LLMTrackedMemoryMap)
 			{
 				for (TMap<FName, uint64>::TConstIterator It(LLMTrackedMemoryMap); It; ++It)
 				{
-					const int SchemaVersion = 1;
+					const int SchemaVersion = 2;
 
-					TArray<FAnalyticsEventAttribute> Attributes;
+					TArray<FAnalyticsEventAttribute> EventAttributes = Attributes;
 
-					Attributes.Emplace(TEXT("SchemaVersion"), SchemaVersion);
-					Attributes.Emplace(TEXT("Context"), Context);
-					Attributes.Emplace(TEXT("TagSet"), TEXT("None"));
-					Attributes.Emplace(TEXT("Name"), It->Key);
-					Attributes.Emplace(TEXT("Size"), It->Value);
+					EventAttributes.Emplace(TEXT("SchemaVersion"), SchemaVersion);
+					EventAttributes.Emplace(TEXT("Context"), Context);
+					EventAttributes.Emplace(TEXT("TagSet"), TagSet);
+					EventAttributes.Emplace(TEXT("Name"), It->Key);
+					EventAttributes.Emplace(TEXT("Size"), It->Value);
 
-					FStudioTelemetry::Get().RecordEvent("Core.Memory.LLM", Attributes);
+					FStudioTelemetry::Get().RecordEvent("Core.Memory.LLM", EventAttributes);
 				}
 			};
 
@@ -72,12 +74,12 @@ void FRuntimeTelemetry::RecordEvent_MemoryLLM(const FString& Context)
 		TMap<FName, uint64> LLMTrackedNoneMemory;
 		FLowLevelMemTracker::Get().GetTrackedTagsNamesWithAmount(LLMTrackedNoneMemory, ELLMTracker::Default, ELLMTagSet::None);
 		RecordLLMMemoryEvent(Context, TEXT("None"), LLMTrackedNoneMemory);
-	
+
 		// AssetClasses TagSet
 		TMap<FName, uint64> LLMTrackedAssetClassesMemory;
 		FLowLevelMemTracker::Get().GetTrackedTagsNamesWithAmount(LLMTrackedAssetClassesMemory, ELLMTracker::Default, ELLMTagSet::AssetClasses);
 		RecordLLMMemoryEvent(Context, TEXT("AssetClasses"), LLMTrackedAssetClassesMemory);
-		
+
 		// Asset TagSet	
 		TMap<FName, uint64> LLMTrackedAssetMemory;
 		FLowLevelMemTracker::Get().GetTrackedTagsNamesWithAmount(LLMTrackedAssetMemory, ELLMTracker::Default, ELLMTagSet::Assets);
@@ -85,7 +87,6 @@ void FRuntimeTelemetry::RecordEvent_MemoryLLM(const FString& Context)
 	}
 #endif
 }
-
 void FRuntimeTelemetry::StartSession()
 {
 	// Register a call back to intercept the point where the game is shutting down
diff --git a/Engine/Plugins/Experimental/RuntimeTelemetry/Source/Public/RuntimeTelemetry.h b/Engine/Plugins/Experimental/RuntimeTelemetry/Source/Public/RuntimeTelemetry.h
--- a/Engine/Plugins/Experimental/RuntimeTelemetry/Source/Public/RuntimeTelemetry.h
+++ b/Engine/Plugins/Experimental/RuntimeTelemetry/Source/Public/RuntimeTelemetry.h
@@ -22,8 +22,8 @@ public:
 	void EndSession();
 
 	/** Useful event recording functions */
-	void RecordEvent_IoStoreOnDemand(const FString& Context);
-	void RecordEvent_MemoryLLM(const FString& Context);
+	void RecordEvent_IoStoreOnDemand(const FString& Context, TArray<FAnalyticsEventAttribute> Attributes = {});
+	void RecordEvent_MemoryLLM(const FString& Context, TArray<FAnalyticsEventAttribute> Attributes = {});
 	
 private:	
 
