Add logging to RuntimeTelemetry so it is easier to debug

#rnx

- Added a new log category 'LogRuntimeTelemetry for use in the module.
- Added a log line when EnginePreExit logging occurs, this will make it easier in the future to confirm that the analytics is indeed taking place.

diff --git a/Engine/Plugins/Experimental/RuntimeTelemetry/Source/Private/RuntimeTelemetry.cpp b/Engine/Plugins/Experimental/RuntimeTelemetry/Source/Private/RuntimeTelemetry.cpp
--- a/Engine/Plugins/Experimental/RuntimeTelemetry/Source/Private/RuntimeTelemetry.cpp
+++ b/Engine/Plugins/Experimental/RuntimeTelemetry/Source/Private/RuntimeTelemetry.cpp
@@ -10,6 +10,7 @@
 #include "IO/IoStoreOnDemand.h"
 #endif
 
+DEFINE_LOG_CATEGORY(LogRuntimeTelemetry);
 
 FRuntimeTelemetry& FRuntimeTelemetry::Get()
 {
@@ -92,6 +93,8 @@ void FRuntimeTelemetry::StartSession()
 	// Register a call back to intercept the point where the game is shutting down
 	FCoreDelegates::OnEnginePreExit.AddLambda([this]()
 		{
+			UE_LOG(LogRuntimeTelemetry, Log, TEXT("Recording EnginePreExit events"));
+
 			RecordEvent_IoStoreOnDemand(TEXT("EnginePreExit"));
 			RecordEvent_MemoryLLM(TEXT("EnginePreExit"));
 		}
diff --git a/Engine/Plugins/Experimental/RuntimeTelemetry/Source/Public/RuntimeTelemetry.h b/Engine/Plugins/Experimental/RuntimeTelemetry/Source/Public/RuntimeTelemetry.h
--- a/Engine/Plugins/Experimental/RuntimeTelemetry/Source/Public/RuntimeTelemetry.h
+++ b/Engine/Plugins/Experimental/RuntimeTelemetry/Source/Public/RuntimeTelemetry.h
@@ -7,6 +7,8 @@
 #include "CoreMinimal.h"
 #include "Templates/SharedPointer.h"
 
+DECLARE_LOG_CATEGORY_EXTERN(LogRuntimeTelemetry, Log, All);
+
 /**
  * 
  */
